Move WorldBuilding plugin out of Restricted folder (make it public)

diff --git a/Engine/Plugins/WorldBuilding/Config/DefaultEngine.ini b/Engine/Plugins/WorldBuilding/Config/DefaultEngine.ini
new file mode 100644
--- /dev/null
+++ b/Engine/Plugins/WorldBuilding/Config/DefaultEngine.ini
@@ -0,0 +1,2 @@
+[/Script/Engine.Engine]
++ActiveClassRedirects=(OldClassName="OldActorType",NewClassName="NewActorType")
\ No newline at end of file
diff --git a/Engine/Plugins/WorldBuilding/Source/WorldBuilding/Private/TestClassRedirectors/NewActorType.cpp b/Engine/Plugins/WorldBuilding/Source/WorldBuilding/Private/TestClassRedirectors/NewActorType.cpp
new file mode 100644
--- /dev/null
+++ b/Engine/Plugins/WorldBuilding/Source/WorldBuilding/Private/TestClassRedirectors/NewActorType.cpp
@@ -0,0 +1,7 @@
+// Copyright Epic Games, Inc. All Rights Reserved.
+
+#include "TestClassRedirectors/NewActorType.h"
+
+ANewActorType::ANewActorType(const FObjectInitializer& ObjectInitializer)
+	: Super(ObjectInitializer)
+{}
\ No newline at end of file
diff --git a/Engine/Plugins/WorldBuilding/Source/WorldBuilding/Private/WorldBuildingCharacterMovement.cpp b/Engine/Plugins/WorldBuilding/Source/WorldBuilding/Private/WorldBuildingCharacterMovement.cpp
new file mode 100644
--- /dev/null
+++ b/Engine/Plugins/WorldBuilding/Source/WorldBuilding/Private/WorldBuildingCharacterMovement.cpp
@@ -0,0 +1,29 @@
+// Copyright Epic Games, Inc. All Rights Reserved.
+
+#include "WorldBuildingCharacterMovement.h"
+#include "WorldBuildingPawn.h"
+
+//----------------------------------------------------------------------//
+// UWorldBuildingCharacterMovement
+//----------------------------------------------------------------------//
+UWorldBuildingCharacterMovement::UWorldBuildingCharacterMovement(const FObjectInitializer& ObjectInitializer)
+	: Super(ObjectInitializer)
+{
+}
+
+
+float UWorldBuildingCharacterMovement::GetMaxSpeed() const
+{
+	float MaxSpeed = Super::GetMaxSpeed();
+
+	const AWorldBuildingPawn* Pawn = Cast<AWorldBuildingPawn>(PawnOwner);
+	if (Pawn)
+	{
+		if (Pawn->IsRunning())
+		{
+			MaxSpeed *= Pawn->GetRunningSpeedModifier();
+		}
+	}
+
+	return MaxSpeed;
+}
diff --git a/Engine/Plugins/WorldBuilding/Source/WorldBuilding/Private/WorldBuildingCharacterMovement.h b/Engine/Plugins/WorldBuilding/Source/WorldBuilding/Private/WorldBuildingCharacterMovement.h
new file mode 100644
--- /dev/null
+++ b/Engine/Plugins/WorldBuilding/Source/WorldBuilding/Private/WorldBuildingCharacterMovement.h
@@ -0,0 +1,14 @@
+// Copyright Epic Games, Inc. All Rights Reserved.
+
+#pragma once
+
+#include "GameFramework/CharacterMovementComponent.h"
+#include "WorldBuildingCharacterMovement.generated.h"
+
+UCLASS()
+class UWorldBuildingCharacterMovement : public UCharacterMovementComponent
+{
+	GENERATED_UCLASS_BODY()
+
+	virtual float GetMaxSpeed() const override;
+};
diff --git a/Engine/Plugins/WorldBuilding/Source/WorldBuilding/Private/WorldBuildingGameMode.cpp b/Engine/Plugins/WorldBuilding/Source/WorldBuilding/Private/WorldBuildingGameMode.cpp
new file mode 100644
--- /dev/null
+++ b/Engine/Plugins/WorldBuilding/Source/WorldBuilding/Private/WorldBuildingGameMode.cpp
@@ -0,0 +1,10 @@
+// Copyright Epic Games, Inc. All Rights Reserved.
+
+#include "WorldBuildingGameMode.h"
+#include "WorldBuildingPawn.h"
+
+AWorldBuildingGameMode::AWorldBuildingGameMode(const FObjectInitializer& ObjectInitializer) : Super(ObjectInitializer)
+{
+	DefaultPawnClass = AWorldBuildingPawn::StaticClass();
+}
+
diff --git a/Engine/Plugins/WorldBuilding/Source/WorldBuilding/Private/WorldBuildingGameMode.h b/Engine/Plugins/WorldBuilding/Source/WorldBuilding/Private/WorldBuildingGameMode.h
new file mode 100644
--- /dev/null
+++ b/Engine/Plugins/WorldBuilding/Source/WorldBuilding/Private/WorldBuildingGameMode.h
@@ -0,0 +1,12 @@
+// Copyright Epic Games, Inc. All Rights Reserved.
+
+#pragma once
+
+#include "GameFramework/GameMode.h"
+#include "WorldBuildingGameMode.generated.h"
+
+UCLASS()
+class AWorldBuildingGameMode : public AGameMode
+{
+	GENERATED_UCLASS_BODY()
+};
diff --git a/Engine/Plugins/WorldBuilding/Source/WorldBuilding/Private/WorldBuildingModule.cpp b/Engine/Plugins/WorldBuilding/Source/WorldBuilding/Private/WorldBuildingModule.cpp
new file mode 100644
--- /dev/null
+++ b/Engine/Plugins/WorldBuilding/Source/WorldBuilding/Private/WorldBuildingModule.cpp
@@ -0,0 +1,15 @@
+// Copyright Epic Games, Inc. All Rights Reserved.
+#include "WorldBuildingModule.h"
+#include "Modules/ModuleManager.h"
+
+IMPLEMENT_MODULE(FWorldBuildingModule, WorldBuilding);
+
+void FWorldBuildingModule::StartupModule()
+{
+
+}
+
+void FWorldBuildingModule::ShutdownModule()
+{
+
+}
\ No newline at end of file
diff --git a/Engine/Plugins/WorldBuilding/Source/WorldBuilding/Private/WorldBuildingModule.h b/Engine/Plugins/WorldBuilding/Source/WorldBuilding/Private/WorldBuildingModule.h
new file mode 100644
--- /dev/null
+++ b/Engine/Plugins/WorldBuilding/Source/WorldBuilding/Private/WorldBuildingModule.h
@@ -0,0 +1,22 @@
+// Copyright Epic Games, Inc. All Rights Reserved.
+#pragma once
+
+#include "CoreMinimal.h"
+#include "IWorldBuildingModule.h"
+
+/**
+ * The module holding all of the UI related pieces for SubLevels management
+ */
+class FWorldBuildingModule : public IWorldBuildingModule
+{
+public:
+	/**
+	 * Called right after the module DLL has been loaded and the module object has been created
+	 */
+	virtual void StartupModule() override;
+
+	/**
+	 * Called before the module is unloaded, right before the module object is destroyed.
+	 */
+	virtual void ShutdownModule() override;
+};
diff --git a/Engine/Plugins/WorldBuilding/Source/WorldBuilding/Private/WorldBuildingPawn.cpp b/Engine/Plugins/WorldBuilding/Source/WorldBuilding/Private/WorldBuildingPawn.cpp
new file mode 100644
--- /dev/null
+++ b/Engine/Plugins/WorldBuilding/Source/WorldBuilding/Private/WorldBuildingPawn.cpp
@@ -0,0 +1,191 @@
+// Copyright Epic Games, Inc. All Rights Reserved.
+
+#include "WorldBuildingPawn.h"
+#include "WorldBuildingCharacterMovement.h"
+#include "UObject/ConstructorHelpers.h"
+#include "Net/UnrealNetwork.h"
+#include "GameFramework/PlayerInput.h"
+#include "GameFramework/Controller.h"
+#include "Engine/StaticMesh.h"
+#include "Components/StaticMeshComponent.h"
+#include "Components/InputComponent.h"
+
+FName AWorldBuildingPawn::MeshComponentName(TEXT("MeshComponent0"));
+
+AWorldBuildingPawn::AWorldBuildingPawn(const FObjectInitializer& ObjectInitializer)
+	: Super(ObjectInitializer.SetDefaultSubobjectClass<UWorldBuildingCharacterMovement>(ACharacter::CharacterMovementComponentName))
+{
+	BaseTurnRate = 45.f;
+	BaseLookUpRate = 45.f;
+	RunningSpeedModifier = 2.5f;
+
+	// Structure to hold one-time initialization
+	struct FConstructorStatics
+	{
+		ConstructorHelpers::FObjectFinder<UStaticMesh> SphereMesh;
+		FConstructorStatics()
+			: SphereMesh(TEXT("/Engine/BasicShapes/Sphere")) {}
+	};
+
+	static FConstructorStatics ConstructorStatics;
+
+	// Visuals only
+	MeshComponent = CreateOptionalDefaultSubobject<UStaticMeshComponent>(AWorldBuildingPawn::MeshComponentName);
+	if (MeshComponent)
+	{
+		MeshComponent->SetStaticMesh(ConstructorStatics.SphereMesh.Object);
+		MeshComponent->AlwaysLoadOnClient = true;
+		MeshComponent->AlwaysLoadOnServer = true;
+		MeshComponent->bOwnerNoSee = true;
+		MeshComponent->bCastDynamicShadow = true;
+		MeshComponent->bAffectDynamicIndirectLighting = false;
+		MeshComponent->bAffectDistanceFieldLighting = false;
+		MeshComponent->bVisibleInRayTracing = false;
+		MeshComponent->PrimaryComponentTick.TickGroup = TG_PrePhysics;
+		MeshComponent->SetupAttachment(RootComponent);
+		// No need for Collision because base class ACharacter as a Capsule component
+		MeshComponent->SetCollisionProfileName(UCollisionProfile::NoCollision_ProfileName);
+		MeshComponent->SetGenerateOverlapEvents(false);
+		MeshComponent->SetCanEverAffectNavigation(false);
+	}
+}
+
+void AWorldBuildingPawn::GetLifetimeReplicatedProps(TArray<FLifetimeProperty>& OutLifetimeProps) const
+{
+	Super::GetLifetimeReplicatedProps(OutLifetimeProps);
+
+	DOREPLIFETIME_CONDITION(AWorldBuildingPawn, bWantsToRun, COND_SkipOwner);
+}
+
+void AWorldBuildingPawn::SetupPlayerInputComponent(UInputComponent* InInputComponent)
+{
+	static bool bBindingsAdded = false;
+	if (!bBindingsAdded)
+	{
+		bBindingsAdded = true;
+
+		// Do our own mapping to not depend on DefaultInput.ini
+		UPlayerInput::AddEngineDefinedAxisMapping(FInputAxisKeyMapping("WorldBuildingPawn_MoveForward", EKeys::W, 1.f));
+		UPlayerInput::AddEngineDefinedAxisMapping(FInputAxisKeyMapping("WorldBuildingPawn_MoveForward", EKeys::S, -1.f));
+		UPlayerInput::AddEngineDefinedAxisMapping(FInputAxisKeyMapping("WorldBuildingPawn_MoveForward", EKeys::Up, 1.f));
+		UPlayerInput::AddEngineDefinedAxisMapping(FInputAxisKeyMapping("WorldBuildingPawn_MoveForward", EKeys::Down, -1.f));
+		UPlayerInput::AddEngineDefinedAxisMapping(FInputAxisKeyMapping("WorldBuildingPawn_MoveForward", EKeys::Gamepad_LeftY, 1.f));
+
+		UPlayerInput::AddEngineDefinedAxisMapping(FInputAxisKeyMapping("WorldBuildingPawn_MoveRight", EKeys::A, -1.f));
+		UPlayerInput::AddEngineDefinedAxisMapping(FInputAxisKeyMapping("WorldBuildingPawn_MoveRight", EKeys::D, 1.f));
+		UPlayerInput::AddEngineDefinedAxisMapping(FInputAxisKeyMapping("WorldBuildingPawn_MoveRight", EKeys::Gamepad_LeftX, 1.f));
+		UPlayerInput::AddEngineDefinedAxisMapping(FInputAxisKeyMapping("WorldBuildingPawn_MoveUp", EKeys::Gamepad_LeftThumbstick, 1.f));
+		UPlayerInput::AddEngineDefinedAxisMapping(FInputAxisKeyMapping("WorldBuildingPawn_MoveUp", EKeys::Gamepad_RightThumbstick, -1.f));
+		UPlayerInput::AddEngineDefinedAxisMapping(FInputAxisKeyMapping("WorldBuildingPawn_MoveUp", EKeys::Gamepad_FaceButton_Bottom, 1.f));
+		UPlayerInput::AddEngineDefinedAxisMapping(FInputAxisKeyMapping("WorldBuildingPawn_MoveUp", EKeys::LeftControl, -1.f));
+		UPlayerInput::AddEngineDefinedAxisMapping(FInputAxisKeyMapping("WorldBuildingPawn_MoveUp", EKeys::SpaceBar, 1.f));
+		UPlayerInput::AddEngineDefinedAxisMapping(FInputAxisKeyMapping("WorldBuildingPawn_MoveUp", EKeys::C, -1.f));
+		UPlayerInput::AddEngineDefinedAxisMapping(FInputAxisKeyMapping("WorldBuildingPawn_MoveUp", EKeys::E, 1.f));
+		UPlayerInput::AddEngineDefinedAxisMapping(FInputAxisKeyMapping("WorldBuildingPawn_MoveUp", EKeys::Q, -1.f));
+
+		UPlayerInput::AddEngineDefinedAxisMapping(FInputAxisKeyMapping("WorldBuildingPawn_TurnRate", EKeys::Gamepad_RightX, 1.f));
+		UPlayerInput::AddEngineDefinedAxisMapping(FInputAxisKeyMapping("WorldBuildingPawn_TurnRate", EKeys::Left, -1.f));
+		UPlayerInput::AddEngineDefinedAxisMapping(FInputAxisKeyMapping("WorldBuildingPawn_TurnRate", EKeys::Right, 1.f));
+		UPlayerInput::AddEngineDefinedAxisMapping(FInputAxisKeyMapping("WorldBuildingPawn_Turn", EKeys::MouseX, 1.f));
+
+		UPlayerInput::AddEngineDefinedAxisMapping(FInputAxisKeyMapping("WorldBuildingPawn_LookUpRate", EKeys::Gamepad_RightY, 1.f));
+		UPlayerInput::AddEngineDefinedAxisMapping(FInputAxisKeyMapping("WorldBuildingPawn_LookUp", EKeys::MouseY, -1.f));
+
+		UPlayerInput::AddEngineDefinedActionMapping(FInputActionKeyMapping("WorldBuildingPawn_Run", EKeys::LeftShift));
+		UPlayerInput::AddEngineDefinedActionMapping(FInputActionKeyMapping("WorldBuildingPawn_Run", EKeys::Gamepad_LeftShoulder));
+	}
+	
+	check(InInputComponent);
+
+	InInputComponent->BindAxis("WorldBuildingPawn_MoveForward", this, &AWorldBuildingPawn::MoveForward);
+	InInputComponent->BindAxis("WorldBuildingPawn_MoveRight", this, &AWorldBuildingPawn::MoveRight);
+	InInputComponent->BindAxis("WorldBuildingPawn_MoveUp", this, &AWorldBuildingPawn::MoveUp_World);
+	InInputComponent->BindAxis("WorldBuildingPawn_Turn", this, &AWorldBuildingPawn::AddControllerYawInput);
+	InInputComponent->BindAxis("WorldBuildingPawn_TurnRate", this, &AWorldBuildingPawn::TurnAtRate);
+	InInputComponent->BindAxis("WorldBuildingPawn_LookUp", this, &AWorldBuildingPawn::AddControllerPitchInput);
+	InInputComponent->BindAxis("WorldBuildingPawn_LookUpRate", this, &AWorldBuildingPawn::LookUpAtRate);
+	InInputComponent->BindAction("WorldBuildingPawn_Run", IE_Pressed, this, &AWorldBuildingPawn::OnStartRunning);
+	InInputComponent->BindAction("WorldBuildingPawn_Run", IE_Released, this, &AWorldBuildingPawn::OnStopRunning);
+}
+
+bool AWorldBuildingPawn::IsRunning() const
+{
+	if (!GetCharacterMovement())
+	{
+		return false;
+	}
+
+	return bWantsToRun && !GetVelocity().IsZero() && (GetVelocity().GetSafeNormal2D() | GetActorForwardVector()) > -0.1;
+}
+
+void AWorldBuildingPawn::MoveRight(float Val)
+{
+	if (Val != 0.f)
+	{
+		if (Controller)
+		{
+			FRotator const ControlSpaceRot = Controller->GetControlRotation();
+
+			// transform to world space and add it
+			AddMovementInput(FRotationMatrix(ControlSpaceRot).GetScaledAxis(EAxis::Y), Val);
+		}
+	}
+}
+
+void AWorldBuildingPawn::MoveForward(float Val)
+{
+	if (Val != 0.f)
+	{
+		if (Controller)
+		{
+			FRotator const ControlSpaceRot = Controller->GetControlRotation();
+
+			// transform to world space and add it
+			AddMovementInput(FRotationMatrix(ControlSpaceRot).GetScaledAxis(EAxis::X), Val);
+		}
+	}
+}
+
+void AWorldBuildingPawn::MoveUp_World(float Val)
+{
+	if (Val != 0.f)
+	{
+		AddMovementInput(FVector::UpVector, Val);
+	}
+}
+
+void AWorldBuildingPawn::TurnAtRate(float Rate)
+{
+	// calculate delta for this frame from the rate information
+	AddControllerYawInput(Rate * BaseTurnRate * GetWorld()->GetDeltaSeconds() * CustomTimeDilation);
+}
+
+void AWorldBuildingPawn::LookUpAtRate(float Rate)
+{
+	// calculate delta for this frame from the rate information
+	AddControllerPitchInput(Rate * BaseLookUpRate * GetWorld()->GetDeltaSeconds() * CustomTimeDilation);
+}
+
+void AWorldBuildingPawn::OnStartRunning()
+{
+	SetRunning(true);
+}
+void AWorldBuildingPawn::OnStopRunning()
+{
+	SetRunning(false);
+}
+
+void AWorldBuildingPawn::SetRunning(bool bNewRunning)
+{
+	bWantsToRun = bNewRunning;
+
+	if (GetLocalRole() < ROLE_Authority)
+	{
+		ServerSetRunning(bNewRunning);
+	}
+}
+
+void AWorldBuildingPawn::ServerSetRunning_Implementation(bool bNewRunning)
+{
+	SetRunning(bNewRunning);
+}
\ No newline at end of file
diff --git a/Engine/Plugins/WorldBuilding/Source/WorldBuilding/Private/WorldBuildingPawn.h b/Engine/Plugins/WorldBuilding/Source/WorldBuilding/Private/WorldBuildingPawn.h
new file mode 100644
--- /dev/null
+++ b/Engine/Plugins/WorldBuilding/Source/WorldBuilding/Private/WorldBuildingPawn.h
@@ -0,0 +1,52 @@
+// Copyright Epic Games, Inc. All Rights Reserved.
+
+#pragma once
+
+#include "GameFramework/Character.h"
+#include "WorldBuildingPawn.generated.h"
+
+class UInputComponent;
+class UStaticMeshComponent;
+
+UCLASS()
+class AWorldBuildingPawn : public ACharacter
+{
+	GENERATED_UCLASS_BODY()
+
+public:
+	virtual void SetupPlayerInputComponent(UInputComponent* InInputComponent) override;
+
+	bool IsRunning() const;
+	float GetRunningSpeedModifier() const { return RunningSpeedModifier; }
+
+protected:
+	void MoveForward(float Val);
+	void MoveRight(float Val);
+	void MoveUp_World(float Val);
+	void TurnAtRate(float Rate);
+	void LookUpAtRate(float Rate);
+
+	void OnStartRunning();
+	void OnStopRunning();
+	void SetRunning(bool bNewRunning);
+
+	UFUNCTION(reliable, server)
+	void ServerSetRunning(bool bNewRunning);
+
+	/** Name of the MeshComponent. Use this name if you want to prevent creation of the component (with ObjectInitializer.DoNotCreateDefaultSubobject). */
+	static FName MeshComponentName;
+
+public:
+	/** The mesh associated with this Pawn. */
+	UPROPERTY()
+	TObjectPtr<UStaticMeshComponent> MeshComponent;
+
+	/** current running state */
+	UPROPERTY(Transient, Replicated)
+	uint8 bWantsToRun : 1;
+
+private:
+	float BaseTurnRate;
+	float BaseLookUpRate;
+	float RunningSpeedModifier;
+};
\ No newline at end of file
diff --git a/Engine/Plugins/WorldBuilding/Source/WorldBuilding/Public/IWorldBuildingModule.h b/Engine/Plugins/WorldBuilding/Source/WorldBuilding/Public/IWorldBuildingModule.h
new file mode 100644
--- /dev/null
+++ b/Engine/Plugins/WorldBuilding/Source/WorldBuilding/Public/IWorldBuildingModule.h
@@ -0,0 +1,13 @@
+// Copyright Epic Games, Inc. All Rights Reserved.
+
+#pragma once
+
+#include "CoreMinimal.h"
+#include "Modules/ModuleInterface.h"
+
+/**
+ * The module holding all WorldBuilding runtime test code
+ */
+class IWorldBuildingModule : public IModuleInterface
+{
+};
\ No newline at end of file
diff --git a/Engine/Plugins/WorldBuilding/Source/WorldBuilding/Public/TestClassRedirectors/NewActorType.h b/Engine/Plugins/WorldBuilding/Source/WorldBuilding/Public/TestClassRedirectors/NewActorType.h
new file mode 100644
--- /dev/null
+++ b/Engine/Plugins/WorldBuilding/Source/WorldBuilding/Public/TestClassRedirectors/NewActorType.h
@@ -0,0 +1,14 @@
+// Copyright Epic Games, Inc. All Rights Reserved.
+
+#pragma once
+
+#include "CoreMinimal.h"
+#include "UObject/ObjectMacros.h"
+#include "GameFramework/Actor.h"
+#include "NewActorType.generated.h"
+
+UCLASS()
+class ANewActorType : public AActor
+{
+	GENERATED_UCLASS_BODY()
+};
\ No newline at end of file
diff --git a/Engine/Plugins/WorldBuilding/Source/WorldBuilding/WorldBuilding.Build.cs b/Engine/Plugins/WorldBuilding/Source/WorldBuilding/WorldBuilding.Build.cs
new file mode 100644
--- /dev/null
+++ b/Engine/Plugins/WorldBuilding/Source/WorldBuilding/WorldBuilding.Build.cs
@@ -0,0 +1,46 @@
+// Copyright Epic Games, Inc. All Rights Reserved.
+
+using UnrealBuildTool;
+
+public class WorldBuilding : ModuleRules
+{
+	public WorldBuilding(ReadOnlyTargetRules Target) : base(Target)
+	{
+		PCHUsage = ModuleRules.PCHUsageMode.UseExplicitOrSharedPCHs;
+	
+		PublicIncludePaths.AddRange(
+			new string[] {
+				// ... add public include paths required here ...
+			}
+			);
+				
+		
+		PrivateIncludePaths.AddRange(
+			new string[] {
+				"WorldBuilding/Private",
+				// ... add other private include paths required here ...
+			}
+			);
+			
+		
+		PublicDependencyModuleNames.AddRange(
+			new string[]
+			{
+				// ... add other public dependencies that you statically link with here ...
+			}
+			);
+			
+		
+		PrivateDependencyModuleNames.AddRange(
+			new string[]
+			{
+				"Core",
+				"InputCore",
+				"CoreUObject",
+				"Engine",
+				"EngineSettings",
+				// ... add private dependencies that you statically link with here ...	
+			}
+			);
+	}
+}
diff --git a/Engine/Plugins/WorldBuilding/WorldBuilding.uplugin b/Engine/Plugins/WorldBuilding/WorldBuilding.uplugin
new file mode 100644
--- /dev/null
+++ b/Engine/Plugins/WorldBuilding/WorldBuilding.uplugin
@@ -0,0 +1,24 @@
+{
+	"FileVersion": 1,
+	"Version": 1,
+	"VersionName": "World Building",
+	"FriendlyName": "World Building",
+	"Description": "UE Plugin for World Building Team.",
+	"Category": "Open World",
+	"CreatedBy": "Epic",
+	"CreatedByURL": "",
+	"DocsURL": "",
+	"MarketplaceURL": "",
+	"SupportURL": "",
+	"EnabledByDefault": false,
+	"CanContainContent": true,
+	"IsBetaVersion": true,
+	"Installed": false,
+	"Modules": [
+		{
+			"Name": "WorldBuilding",
+			"Type": "Runtime",
+			"LoadingPhase": "Default"
+		}
+	]
+}
\ No newline at end of file
