Update WorldBuilding test map for class redirectors test.

diff --git a/Engine/Plugins/WorldBuilding/Config/DefaultEngine.ini b/Engine/Plugins/WorldBuilding/Config/DefaultEngine.ini
--- a/Engine/Plugins/WorldBuilding/Config/DefaultEngine.ini
+++ b/Engine/Plugins/WorldBuilding/Config/DefaultEngine.ini
@@ -1,2 +1,2 @@
 [/Script/Engine.Engine]
-+ActiveClassRedirects=(OldClassName="OldActorType",NewClassName="NewActorType")
\ No newline at end of file
++ActiveClassRedirects=(OldClassName="/Script/WorldBuilding.OldActorType",NewClassName="/Script/WorldBuilding.NewActorType")
\ No newline at end of file
