WorldBuilding Replay Spectator

diff --git a/Engine/Plugins/WorldBuilding/Source/WorldBuilding/Private/Replay/WorldBuildingReplaySpectator.cpp b/Engine/Plugins/WorldBuilding/Source/WorldBuilding/Private/Replay/WorldBuildingReplaySpectator.cpp
new file mode 100644
--- /dev/null
+++ b/Engine/Plugins/WorldBuilding/Source/WorldBuilding/Private/Replay/WorldBuildingReplaySpectator.cpp
@@ -0,0 +1,97 @@
+// Copyright Epic Games, Inc. All Rights Reserved.
+
+#include "WorldBuildingReplaySpectator.h"
+#include "WorldPartition/WorldPartition.h"
+#include "WorldBuildingReplaySpectatorHUD.h"
+#include "Engine/DemoNetDriver.h"
+#include "GameFramework/PlayerInput.h"
+
+AWorldBuildingReplaySpectator::AWorldBuildingReplaySpectator(const FObjectInitializer& ObjectInitializer) : Super(ObjectInitializer)
+{
+}
+
+void AWorldBuildingReplaySpectator::GetPlayerViewPoint(FVector& out_Location, FRotator& out_Rotation) const
+{
+	if (AWorldBuildingReplaySpectatorHUD* HUD = Cast<AWorldBuildingReplaySpectatorHUD>(MyHUD))
+	{
+		if (UWorldPartition* WorldPartition = GetWorld()->GetWorldPartition())
+		{
+			const TArray<FWorldPartitionStreamingSource>& StreamingSources = WorldPartition->GetStreamingSources();
+			FString SelectedSourceName = HUD->GetSelectedViewPoint();
+			int32 Index = StreamingSources.IndexOfByPredicate([SelectedSourceName](const FWorldPartitionStreamingSource& StreamingSource)
+				{
+					return SelectedSourceName == StreamingSource.Name.ToString();
+				});
+
+			if (Index != INDEX_NONE)
+			{
+				out_Location = WorldPartition->GetInstanceTransform().TransformPosition(StreamingSources[Index].Location);
+				out_Rotation = WorldPartition->GetInstanceTransform().TransformRotation(StreamingSources[Index].Rotation.Quaternion()).Rotator();
+				return;
+			}
+		}
+	}
+
+	// FREE CAM
+	return Super::GetPlayerViewPoint(out_Location, out_Rotation);
+}
+
+void AWorldBuildingReplaySpectator::SpawnDefaultHUD()
+{
+	FActorSpawnParameters SpawnParams;
+	SpawnParams.Owner = this;
+	SpawnParams.Instigator = GetInstigator();
+	SpawnParams.ObjectFlags |= RF_Transient;	// We never want to save HUDs into a map
+	MyHUD = GetWorld()->SpawnActor<AWorldBuildingReplaySpectatorHUD>(HUDClass ? HUDClass.Get() : AWorldBuildingReplaySpectatorHUD::StaticClass(), SpawnParams);
+}
+
+void AWorldBuildingReplaySpectator::TickActor(float DeltaSeconds, ELevelTick TickType, FActorTickFunction& ThisTickFunction)
+{
+	Super::TickActor(DeltaSeconds, TickType, ThisTickFunction);
+
+	if (AWorldBuildingReplaySpectatorHUD* HUD = Cast<AWorldBuildingReplaySpectatorHUD>(MyHUD))
+	{
+		if (UWorldPartition* WorldPartition = GetWorld()->GetWorldPartition())
+		{
+			const TArray<FWorldPartitionStreamingSource>& StreamingSources = WorldPartition->GetStreamingSources();
+			TArray<FString> ViewPointNames;
+			ViewPointNames.Reserve(StreamingSources.Num() + 1);
+			Algo::Transform(StreamingSources, ViewPointNames, [](const FWorldPartitionStreamingSource& StreamingSource) { return StreamingSource.Name.ToString(); });
+			ViewPointNames.Add(FString("FREE CAM"));
+			HUD->SetAvailableViewPoints(ViewPointNames);
+		}
+	}
+}
+
+void AWorldBuildingReplaySpectator::SetupInputComponent()
+{
+	Super::SetupInputComponent();
+
+	static bool bBindingsAdded = false;
+	if (!bBindingsAdded)
+	{
+		bBindingsAdded = true;
+
+		UPlayerInput::AddEngineDefinedActionMapping(FInputActionKeyMapping("ReplaySpectator_NextViewPoint", EKeys::Up));
+		UPlayerInput::AddEngineDefinedActionMapping(FInputActionKeyMapping("ReplaySpectator_PreviousViewPoint", EKeys::Down));
+	}
+
+	InputComponent->BindAction("ReplaySpectator_NextViewPoint", IE_Pressed, this, &AWorldBuildingReplaySpectator::OnNextViewPoint);
+	InputComponent->BindAction("ReplaySpectator_PreviousViewPoint", IE_Pressed, this, &AWorldBuildingReplaySpectator::OnPreviousViewPoint);
+}
+
+void AWorldBuildingReplaySpectator::OnNextViewPoint()
+{
+	if (AWorldBuildingReplaySpectatorHUD* HUD = Cast<AWorldBuildingReplaySpectatorHUD>(MyHUD))
+	{
+		HUD->NextViewPoint();
+	}
+}
+
+void AWorldBuildingReplaySpectator::OnPreviousViewPoint()
+{
+	if (AWorldBuildingReplaySpectatorHUD* HUD = Cast<AWorldBuildingReplaySpectatorHUD>(MyHUD))
+	{
+		HUD->PreviousViewPoint();
+	}
+}
diff --git a/Engine/Plugins/WorldBuilding/Source/WorldBuilding/Private/Replay/WorldBuildingReplaySpectator.h b/Engine/Plugins/WorldBuilding/Source/WorldBuilding/Private/Replay/WorldBuildingReplaySpectator.h
new file mode 100644
--- /dev/null
+++ b/Engine/Plugins/WorldBuilding/Source/WorldBuilding/Private/Replay/WorldBuildingReplaySpectator.h
@@ -0,0 +1,31 @@
+// Copyright Epic Games, Inc. All Rights Reserved.
+
+#pragma once
+
+#include "GameFramework/PlayerController.h"
+#include "Replay/WorldBuildingReplaySpectatorHUD.h"
+#include "WorldBuildingReplaySpectator.generated.h"
+
+/** 
+ * AWorldBuildingReplaySpectator can be used as the replay player controller when setting up a game mode. 
+ * 
+ * Allows changing view point based on World Partition replay streaming sources.
+ */
+UCLASS(config = Game)
+class AWorldBuildingReplaySpectator : public APlayerController
+{
+	GENERATED_UCLASS_BODY()
+
+public:
+	virtual void GetPlayerViewPoint(FVector& out_Location, FRotator& out_Rotation) const override;
+	virtual void SpawnDefaultHUD() override;
+	virtual void TickActor(float DeltaSeconds, ELevelTick TickType, FActorTickFunction& ThisTickFunction) override;
+	virtual void SetupInputComponent() override;
+
+private:
+	void OnNextViewPoint();
+	void OnPreviousViewPoint();
+
+	UPROPERTY(EditAnywhere, Category = HUD)
+	TSubclassOf<AWorldBuildingReplaySpectatorHUD> HUDClass;
+};
\ No newline at end of file
diff --git a/Engine/Plugins/WorldBuilding/Source/WorldBuilding/Private/Replay/WorldBuildingReplaySpectatorHUD.cpp b/Engine/Plugins/WorldBuilding/Source/WorldBuilding/Private/Replay/WorldBuildingReplaySpectatorHUD.cpp
new file mode 100644
--- /dev/null
+++ b/Engine/Plugins/WorldBuilding/Source/WorldBuilding/Private/Replay/WorldBuildingReplaySpectatorHUD.cpp
@@ -0,0 +1,128 @@
+// Copyright Epic Games, Inc. All Rights Reserved.
+
+#include "WorldBuildingReplaySpectatorHUD.h"
+#include "Engine/Canvas.h"
+#include "Engine/Engine.h"
+#include "Components/ComboBoxString.h"
+
+UWorldBuildingViewPointWidget::UWorldBuildingViewPointWidget(const FObjectInitializer& ObjectInitializer) : Super(ObjectInitializer)
+{
+
+}
+
+void UWorldBuildingViewPointWidget::NativeConstruct()
+{
+	Super::NativeConstruct();
+}
+
+void UWorldBuildingViewPointWidget::NextViewPoint()
+{
+	if(ComboViewPoints && ComboViewPoints->GetOptionCount())
+	{
+		SetSelectedViewPointIndex((ComboViewPoints->GetSelectedIndex() + 1) % ComboViewPoints->GetOptionCount());
+	}
+}
+
+void UWorldBuildingViewPointWidget::PreviousViewPoint()
+{
+	if (ComboViewPoints)
+	{
+		const int32 Previous = ComboViewPoints->GetSelectedIndex() - 1;
+		SetSelectedViewPointIndex(Previous < 0 ? ComboViewPoints->GetOptionCount() - 1 : Previous);
+	}
+}
+
+void UWorldBuildingViewPointWidget::SetSelectedViewPointIndex(int32 Index)
+{
+	check(ComboViewPoints);
+	if (Index >= 0 && Index < ComboViewPoints->GetOptionCount())
+	{
+		ComboViewPoints->SetSelectedIndex(Index);
+	}
+}
+
+FString UWorldBuildingViewPointWidget::GetSelectedViewPoint() const
+{
+	if (ComboViewPoints)
+	{
+		return ComboViewPoints->GetSelectedOption();
+	}
+
+	return FString();
+}
+
+void UWorldBuildingViewPointWidget::SetAvailableViewPoints(const TArray<FString>& ViewPoints)
+{
+	if (ComboViewPoints)
+	{
+		if (LastViewPoints != ViewPoints)
+		{
+			FString CurrentSelection = ComboViewPoints->GetSelectedOption();
+			ComboViewPoints->ClearSelection();
+			ComboViewPoints->ClearOptions();
+			for (const FString& Option : ViewPoints)
+			{
+				ComboViewPoints->AddOption(Option);
+			}
+
+			int32 OptionIndex = ComboViewPoints->FindOptionIndex(CurrentSelection);
+			SetSelectedViewPointIndex(OptionIndex != INDEX_NONE ? OptionIndex : 0);
+
+			LastViewPoints = ViewPoints;
+		}
+	}
+}
+
+AWorldBuildingReplaySpectatorHUD::AWorldBuildingReplaySpectatorHUD(const FObjectInitializer& ObjectInitializer) 
+	: Super(ObjectInitializer)
+{
+}
+
+void AWorldBuildingReplaySpectatorHUD::BeginPlay()
+{
+	Super::BeginPlay();
+
+	if (ViewPointWidgetClass)
+	{
+		ViewPointWidget = CreateWidget<UWorldBuildingViewPointWidget>(GetWorld(), ViewPointWidgetClass);
+
+		if (ViewPointWidget)
+		{
+			ViewPointWidget->AddToViewport();
+		}
+	}
+}
+
+FString AWorldBuildingReplaySpectatorHUD::GetSelectedViewPoint() const
+{
+	if (ViewPointWidget)
+	{
+		return ViewPointWidget->GetSelectedViewPoint();
+	}
+
+	return FString();
+}
+
+void AWorldBuildingReplaySpectatorHUD::SetAvailableViewPoints(const TArray<FString>& ViewPoints)
+{
+	if (ViewPointWidget)
+	{
+		ViewPointWidget->SetAvailableViewPoints(ViewPoints);
+	}
+}
+
+void AWorldBuildingReplaySpectatorHUD::NextViewPoint()
+{
+	if (ViewPointWidget)
+	{
+		ViewPointWidget->NextViewPoint();
+	}
+}
+
+void AWorldBuildingReplaySpectatorHUD::PreviousViewPoint()
+{
+	if (ViewPointWidget)
+	{
+		ViewPointWidget->PreviousViewPoint();
+	}
+}
diff --git a/Engine/Plugins/WorldBuilding/Source/WorldBuilding/Private/Replay/WorldBuildingReplaySpectatorHUD.h b/Engine/Plugins/WorldBuilding/Source/WorldBuilding/Private/Replay/WorldBuildingReplaySpectatorHUD.h
new file mode 100644
--- /dev/null
+++ b/Engine/Plugins/WorldBuilding/Source/WorldBuilding/Private/Replay/WorldBuildingReplaySpectatorHUD.h
@@ -0,0 +1,50 @@
+// Copyright Epic Games, Inc. All Rights Reserved.
+
+#pragma once
+
+#include "GameFramework/HUD.h"
+#include "Blueprint/UserWidget.h"
+#include "WorldBuildingReplaySpectatorHUD.generated.h"
+
+UCLASS()
+class UWorldBuildingViewPointWidget : public UUserWidget
+{
+	GENERATED_UCLASS_BODY()
+
+public:
+	virtual void NativeConstruct() override;
+
+	FString GetSelectedViewPoint() const;
+	void SetAvailableViewPoints(const TArray<FString>& ViewPoints);
+
+	void NextViewPoint();
+	void PreviousViewPoint();
+
+private:
+	void SetSelectedViewPointIndex(int32 Index);
+
+	UPROPERTY(meta = (BindWidget))
+	class UComboBoxString* ComboViewPoints;
+
+	TArray<FString> LastViewPoints;
+};
+
+UCLASS()
+class AWorldBuildingReplaySpectatorHUD : public AHUD
+{
+	GENERATED_UCLASS_BODY()
+
+	virtual void BeginPlay() override;
+
+	FString GetSelectedViewPoint() const;
+	void SetAvailableViewPoints(const TArray<FString>& StreamingSources);
+
+	void NextViewPoint();
+	void PreviousViewPoint();
+
+protected:
+	UPROPERTY(EditAnywhere, Category = "Widget")
+	TSubclassOf<UWorldBuildingViewPointWidget> ViewPointWidgetClass;
+
+	UWorldBuildingViewPointWidget* ViewPointWidget;
+};
\ No newline at end of file
diff --git a/Engine/Plugins/WorldBuilding/Source/WorldBuilding/WorldBuilding.Build.cs b/Engine/Plugins/WorldBuilding/Source/WorldBuilding/WorldBuilding.Build.cs
--- a/Engine/Plugins/WorldBuilding/Source/WorldBuilding/WorldBuilding.Build.cs
+++ b/Engine/Plugins/WorldBuilding/Source/WorldBuilding/WorldBuilding.Build.cs
@@ -39,6 +39,7 @@ public class WorldBuilding : ModuleRules
 				"CoreUObject",
 				"Engine",
 				"EngineSettings",
+				"UMG",
 				// ... add private dependencies that you statically link with here ...	
 			}
 			);
diff --git a/Engine/Source/Runtime/Engine/Public/WorldPartition/WorldPartition.h b/Engine/Source/Runtime/Engine/Public/WorldPartition/WorldPartition.h
--- a/Engine/Source/Runtime/Engine/Public/WorldPartition/WorldPartition.h
+++ b/Engine/Source/Runtime/Engine/Public/WorldPartition/WorldPartition.h
@@ -121,12 +121,6 @@ private:
 
 	static void WorldPartitionOnLevelRemovedFromWorld(class ULevel* Level, UWorld* InWorld);
 
-	// UWorldPartitionSubsystem interface+
-	friend class UWorldPartitionSubsystem;
-
-	const TArray<FWorldPartitionStreamingSource>& GetStreamingSources() const;
-	// UActorPartitionSubsystem interface-
-
 public:
 	//~ Begin UObject Interface
 	virtual bool ResolveSubobject(const TCHAR* SubObjectPath, UObject*& OutObject, bool bLoadIfExists) override;
@@ -178,6 +172,8 @@ public:
 	class ULevel* GetPreferredLoadedLevelToAddToWorld() const;
 	bool IsStreamingCompleted(EWorldPartitionRuntimeCellState QueryState, const TArray<FWorldPartitionStreamingQuerySource>& QuerySources, bool bExactState) const;
 
+	const TArray<FWorldPartitionStreamingSource>& GetStreamingSources() const;
+
 	void RegisterStreamingSourceProvider(IWorldPartitionStreamingSourceProvider* StreamingSource);
 	bool UnregisterStreamingSourceProvider(IWorldPartitionStreamingSourceProvider* StreamingSource);
 
